Mirroring CL 4078631 from Dev-VR

Fixing MR Calibration so it scales the alignment model according the the capture's FOV (so they appear the same size across capture devices - leading to a homogenous experience).

Also moved the FOV override config setting to be a console command/setting (mrc.FovOverride) to help in testing this.

diff --git a/Engine/Plugins/Runtime/MixedRealityCaptureFramework/Source/MixedRealityCaptureFramework/Private/MixedRealityCaptureComponent.cpp b/Engine/Plugins/Runtime/MixedRealityCaptureFramework/Source/MixedRealityCaptureFramework/Private/MixedRealityCaptureComponent.cpp
--- a/Engine/Plugins/Runtime/MixedRealityCaptureFramework/Source/MixedRealityCaptureFramework/Private/MixedRealityCaptureComponent.cpp
+++ b/Engine/Plugins/Runtime/MixedRealityCaptureFramework/Source/MixedRealityCaptureFramework/Private/MixedRealityCaptureComponent.cpp
@@ -43,90 +43,68 @@ DEFINE_LOG_CATEGORY(LogMixedRealityCapture);
 
 DECLARE_MULTICAST_DELEGATE(FOnCommandValueChanged);
 
-template<typename T>
-class FMulticastCVarCommand : private FAutoConsoleCommand, public FOnCommandValueChanged
+template<typename T, typename U=T>
+class TMulticastCVarCommand : private TAutoConsoleVariable<U>, public FOnCommandValueChanged
 {
 public:
-	FMulticastCVarCommand(const TCHAR* Name, const TCHAR* Help, T& InVarRef)
-		: FAutoConsoleCommand(Name, Help, FConsoleCommandWithWorldArgsAndOutputDeviceDelegate::CreateRaw(this, &FMulticastCVarCommand<T>::HandleCommand))
-		, VarRef(InVarRef)
-		, CmdName(Name)
-	{}
-
-	void HandleCommand(const TArray<FString>& Args, UWorld*, FOutputDevice& Ar) const
+	TMulticastCVarCommand(const TCHAR* Name, const T& DefaultVal, const TCHAR* Help)
+		: TAutoConsoleVariable<U>(Name, (U)DefaultVal, Help)
 	{
-		if (Args.Num())
-		{
-			T NewVal;
-			LexFromString(NewVal, *Args[0]);
-
-			if (NewVal != VarRef)
-			{
-				VarRef = NewVal;
-				Broadcast();
-			}
-		}
-		Ar.Logf(TEXT("%s = %s"), *CmdName, *LexToString(VarRef));
+		AsVariable()->SetOnChangedCallback(FConsoleVariableDelegate::CreateRaw(this, &TMulticastCVarCommand<T,U>::OnChanged));
 	}
 
-private:
-	T& VarRef;
-	FString CmdName;
-};
+	void OnChanged(IConsoleVariable* /*This*/)
+	{
+		Broadcast();
+	}
 
-//------------------------------------------------------------------------------
-template<>
-void FMulticastCVarCommand<bool>::HandleCommand(const TArray<FString>& Args, UWorld*, FOutputDevice& Ar) const
-{
-	if (Args.Num())
+	T GetValue() const
 	{
-		bool NewVal;
-		LexFromString(NewVal, *Args[0]);
+		return (T)GetValueOnGameThread();
+	}
 
-		if (NewVal != VarRef)
-		{
-			VarRef = NewVal;
-			Broadcast();
-		}
+	operator T() const
+	{
+		return GetValue();
 	}
-// 	else
-// 	{
-// 		VarRef = !VarRef;
-// 	}
-	Ar.Logf(TEXT("%s = %s"), *CmdName, VarRef ? TEXT("On") : TEXT("Off"));
-}
+};
+
+typedef TMulticastCVarCommand<bool, int32> FMulticastBoolCVar;
+typedef TMulticastCVarCommand<float> FMulticastFloatCVar;
 
 /* MRCaptureComponent_Impl
  *****************************************************************************/
 
 namespace MRCaptureComponent_Impl
 {
-	static bool bUseUndistortion = true;
-	static FMulticastCVarCommand<bool> UseUndistortionCommand(
+	static FMulticastBoolCVar UseUndistortion(
 		TEXT("mrc.undistortion"),
-		TEXT("Enables/Disables the undistortion pass for MixedRealityCaptures. When disabled, the default (black) texture is used instead for the distortion displacement map."),
-		bUseUndistortion
+		true,
+		TEXT("Enables/Disables the undistortion pass for MixedRealityCaptures. When disabled, the default (black) texture is used instead for the distortion displacement map.")
 	);
 
-	static bool bUseDistortionFocalAspect = true;
-	static FMulticastCVarCommand<bool> UseFocalLenAspectCommand(
+	static FMulticastBoolCVar UseFocalLenAspect(
 		TEXT("mrc.undistortion.bUseFocalAspectRatio"),
-		TEXT("When enabled, to account for stretching from the OpenCV undistortion process, MixedRealityCaptures will scale their projected aspect ratio by the estimated focal length ratio (as reported by OpenCV)."),
-		bUseDistortionFocalAspect
+		true,
+		TEXT("When enabled, to account for stretching from the OpenCV undistortion process, MixedRealityCaptures will scale their projected aspect ratio by the estimated focal length ratio (as reported by OpenCV).")
 	);
 
-	static float DistortionCroppingAmount = 0.0f;
-	static FMulticastCVarCommand<float> DistortionCroppingAmountCommand(
+	static FMulticastFloatCVar DistortionCroppingAmount(
 		TEXT("mrc.undistortion.CroppingAmount"),
-		TEXT("A value meant to range from 0 to 1. At one, as part of the undistortion process, OpenCV will attempt to crop out all empty pixels resulting from the process (essentially zooming the image). Zero means no cropping will occur."),
-		DistortionCroppingAmount
+		0.0f,
+		TEXT("A value meant to range from 0 to 1. At one, as part of the undistortion process, OpenCV will attempt to crop out all empty pixels resulting from the process (essentially zooming the image). Zero means no cropping will occur.")
 	);
 
-	static bool bUseUndistortedFOV = true;
-	static FMulticastCVarCommand<bool> UseUndistortedFOVCommand(
+	static FMulticastBoolCVar UseUndistortedFOV(
 		TEXT("mrc.undistortion.bUseUndistortedFOV"),
-		TEXT("When enabled, MixedRealityCaptures (MRCs) will use the estimated FOV from the OpenCV undistortion process instead of the FOV the MRC was calibrated with. This accounts for any cropping, etc. done by OpenCV."),
-		bUseUndistortedFOV
+		true,
+		TEXT("When enabled, MixedRealityCaptures (MRCs) will use the estimated FOV from the OpenCV undistortion process instead of the FOV the MRC was calibrated with. This accounts for any cropping, etc. done by OpenCV.")
+	);
+
+	static FMulticastFloatCVar CaptureFOVOverride(
+		TEXT("mrc.FovOverride"),
+		0.0f,
+		TEXT("When set to be greater than zero, MixedRealityCaptures will use this for the FOV instead of what was previously set.")
 	);
 
 	static const FName DefaultDistortionMapParamName(TEXT("DistortionDisplacementMap"));
@@ -173,10 +151,11 @@ namespace MRCaptureComponent_Impl
 //------------------------------------------------------------------------------
 static void MRCaptureComponent_Impl::RemoveAllCVarBindings(const void* BoundObject)
 {
-	UseUndistortionCommand.RemoveAll(BoundObject);
-	UseFocalLenAspectCommand.RemoveAll(BoundObject);
-	DistortionCroppingAmountCommand.RemoveAll(BoundObject);
-	UseUndistortedFOVCommand.RemoveAll(BoundObject);
+	UseUndistortion.RemoveAll(BoundObject);
+	UseFocalLenAspect.RemoveAll(BoundObject);
+	DistortionCroppingAmount.RemoveAll(BoundObject);
+	UseUndistortedFOV.RemoveAll(BoundObject);
+	CaptureFOVOverride.RemoveAll(BoundObject);
 }
 
 //------------------------------------------------------------------------------
@@ -443,10 +422,11 @@ void UMixedRealityCaptureComponent::Activate(bool bReset)
 
 		RefreshCameraFeed();
 
-		MRCaptureComponent_Impl::UseUndistortedFOVCommand.AddUObject(this, &UMixedRealityCaptureComponent::RefreshFOV);
-		MRCaptureComponent_Impl::DistortionCroppingAmountCommand.AddUObject(this, &UMixedRealityCaptureComponent::RefreshDistortionDisplacementMap);
-		MRCaptureComponent_Impl::UseFocalLenAspectCommand.AddUObject(this, &UMixedRealityCaptureComponent::RefreshProjectionDimensions);
-		MRCaptureComponent_Impl::UseUndistortionCommand.AddUObject(this, &UMixedRealityCaptureComponent::RefreshDistortionDisplacementMap);
+		MRCaptureComponent_Impl::CaptureFOVOverride.AddUObject(this, &UMixedRealityCaptureComponent::RefreshFOV);
+		MRCaptureComponent_Impl::UseUndistortedFOV.AddUObject(this, &UMixedRealityCaptureComponent::RefreshFOV);
+		MRCaptureComponent_Impl::DistortionCroppingAmount.AddUObject(this, &UMixedRealityCaptureComponent::RefreshDistortionDisplacementMap);
+		MRCaptureComponent_Impl::UseFocalLenAspect.AddUObject(this, &UMixedRealityCaptureComponent::RefreshProjectionDimensions);
+		MRCaptureComponent_Impl::UseUndistortion.AddUObject(this, &UMixedRealityCaptureComponent::RefreshDistortionDisplacementMap);
 
 		FXRTrackingSystemDelegates::OnXRTrackingOriginChanged.AddUObject(this, &UMixedRealityCaptureComponent::OnTrackingOriginChanged);
 	}	
@@ -903,7 +883,7 @@ float UMixedRealityCaptureComponent::GetDesiredAspectRatio() const
 		}
 	}
 
-	if (MRCaptureComponent_Impl::bUseUndistortion && MRCaptureComponent_Impl::bUseDistortionFocalAspect && LensDistortionParameters.IsSet() && LensFocalLenRatio > 0)
+	if (MRCaptureComponent_Impl::UseUndistortion && MRCaptureComponent_Impl::UseFocalLenAspect && LensDistortionParameters.IsSet() && LensFocalLenRatio > 0)
 	{ 
 		DesiredAspectRatio *= LensFocalLenRatio;
 	}
@@ -914,7 +894,7 @@ float UMixedRealityCaptureComponent::GetDesiredAspectRatio() const
 //------------------------------------------------------------------------------
 void UMixedRealityCaptureComponent::RefreshDistortionDisplacementMap()
 {
-	if (MRCaptureComponent_Impl::bUseUndistortion && LensDistortionParameters.IsSet() && TextureTarget)
+	if (MRCaptureComponent_Impl::UseUndistortion && LensDistortionParameters.IsSet() && TextureTarget)
 	{
 		// for OpenCv, 0.0 is cropped fully, whereas 1.0 is uncropped
 		float CroppingInterpoler = 1.f - MRCaptureComponent_Impl::DistortionCroppingAmount;
@@ -929,11 +909,11 @@ void UMixedRealityCaptureComponent::RefreshDistortionDisplacementMap()
 	}
 	MRCaptureComponent_Impl::ApplyDistortionMapToMaterial(VideoProcessingMaterial, DistortionDisplacementMap);
 
-	if (MRCaptureComponent_Impl::bUseDistortionFocalAspect)
+	if (MRCaptureComponent_Impl::UseFocalLenAspect)
 	{
 		RefreshProjectionDimensions();
 	}
-	if (MRCaptureComponent_Impl::bUseUndistortedFOV)
+	if (MRCaptureComponent_Impl::UseUndistortedFOV)
 	{
 		RefreshFOV();
 	}
@@ -942,12 +922,11 @@ void UMixedRealityCaptureComponent::RefreshDistortionDisplacementMap()
 //------------------------------------------------------------------------------
 void UMixedRealityCaptureComponent::RefreshFOV()
 {
-	float FOVOverride = FOVAngle;
-	if (GConfig->GetFloat(TEXT("/Script/MixedRealityCaptureFramework.MixedRealityFrameworkSettings"), TEXT("CalibratedFOVOverride"), FOVOverride, GEngineIni))
+	if (MRCaptureComponent_Impl::CaptureFOVOverride > 0.0f)
 	{
-		FOVAngle = GetDefault<UMrcFrameworkSettings>()->CalibratedFOVOverride;
+		FOVAngle = MRCaptureComponent_Impl::CaptureFOVOverride;
 	}
-	else if (MRCaptureComponent_Impl::bUseUndistortion && MRCaptureComponent_Impl::bUseUndistortedFOV && LensDistortionParameters.IsSet() && UndistortedFOV > 0)
+	else if (MRCaptureComponent_Impl::UseUndistortion && MRCaptureComponent_Impl::UseUndistortedFOV && LensDistortionParameters.IsSet() && UndistortedFOV > 0)
 	{
 		FOVAngle = UndistortedFOV;
 	}
diff --git a/Engine/Plugins/Runtime/MixedRealityCaptureFramework/Source/MixedRealityCaptureFramework/Private/MrcFrameworkSettings.cpp b/Engine/Plugins/Runtime/MixedRealityCaptureFramework/Source/MixedRealityCaptureFramework/Private/MrcFrameworkSettings.cpp
--- a/Engine/Plugins/Runtime/MixedRealityCaptureFramework/Source/MixedRealityCaptureFramework/Private/MrcFrameworkSettings.cpp
+++ b/Engine/Plugins/Runtime/MixedRealityCaptureFramework/Source/MixedRealityCaptureFramework/Private/MrcFrameworkSettings.cpp
@@ -11,5 +11,4 @@ UMrcFrameworkSettings::UMrcFrameworkSettings(const FObjectInitializer& ObjectIni
 	, DefaulGarbageMatteMesh(TEXT("/MixedRealityCaptureFramework/SM_GarbageMattePlane"))
 	, DefaulGarbageMatteMaterial(TEXT("/MixedRealityCaptureFramework/M_GarbageMatte"))
 	, DefaulGarbageMatteTarget(TEXT("/MixedRealityCaptureFramework/T_MrcGarbageMatteRenderTarget"))
-	, CalibratedFOVOverride(90.f)
 {}
diff --git a/Engine/Plugins/Runtime/MixedRealityCaptureFramework/Source/MixedRealityCaptureFramework/Private/MrcFrameworkSettings.h b/Engine/Plugins/Runtime/MixedRealityCaptureFramework/Source/MixedRealityCaptureFramework/Private/MrcFrameworkSettings.h
--- a/Engine/Plugins/Runtime/MixedRealityCaptureFramework/Source/MixedRealityCaptureFramework/Private/MrcFrameworkSettings.h
+++ b/Engine/Plugins/Runtime/MixedRealityCaptureFramework/Source/MixedRealityCaptureFramework/Private/MrcFrameworkSettings.h
@@ -32,7 +32,4 @@ public:
 
 	UPROPERTY(Config)
 	FSoftObjectPath DefaulGarbageMatteTarget;
-
-	UPROPERTY(Config, meta=(UIMin = "5.0", UIMax = "170", ClampMin = "0.001", ClampMax = "360.0"))
-	float CalibratedFOVOverride;
 };
