Mirroring CL 4080119 from Dev-VR

CIS fix - circumventing vs2015 errors/warnings

#rnx

diff --git a/Engine/Plugins/Runtime/MixedRealityCaptureFramework/Source/MixedRealityCaptureFramework/Private/MixedRealityCaptureComponent.cpp b/Engine/Plugins/Runtime/MixedRealityCaptureFramework/Source/MixedRealityCaptureFramework/Private/MixedRealityCaptureComponent.cpp
--- a/Engine/Plugins/Runtime/MixedRealityCaptureFramework/Source/MixedRealityCaptureFramework/Private/MixedRealityCaptureComponent.cpp
+++ b/Engine/Plugins/Runtime/MixedRealityCaptureFramework/Source/MixedRealityCaptureFramework/Private/MixedRealityCaptureComponent.cpp
@@ -50,7 +50,7 @@ public:
 	TMulticastCVarCommand(const TCHAR* Name, const T& DefaultVal, const TCHAR* Help)
 		: TAutoConsoleVariable<U>(Name, (U)DefaultVal, Help)
 	{
-		AsVariable()->SetOnChangedCallback(FConsoleVariableDelegate::CreateRaw(this, &TMulticastCVarCommand<T,U>::OnChanged));
+		TAutoConsoleVariable<U>::AsVariable()->SetOnChangedCallback(FConsoleVariableDelegate::CreateRaw(this, &TMulticastCVarCommand<T,U>::OnChanged));
 	}
 
 	void OnChanged(IConsoleVariable* /*This*/)
@@ -60,7 +60,7 @@ public:
 
 	T GetValue() const
 	{
-		return (T)GetValueOnGameThread();
+		return (T)TAutoConsoleVariable<U>::GetValueOnGameThread();
 	}
 
 	operator T() const
@@ -69,6 +69,13 @@ public:
 	}
 };
 
+// specialization to avoid warning C4800 ('int32': forcing value to bool 'true' or 'false')
+template<>
+bool TMulticastCVarCommand<bool, int32>::GetValue() const
+{
+	return GetValueOnGameThread() != 0;
+}
+
 typedef TMulticastCVarCommand<bool, int32> FMulticastBoolCVar;
 typedef TMulticastCVarCommand<float> FMulticastFloatCVar;
 
