Mirroring CL 4084289 from Dev-VR

MR - Setting up the calibration save data so that we version it in the future (know that a lens param change is inbound after 4.20).

diff --git a/Engine/Plugins/Runtime/MixedRealityCaptureFramework/Source/MixedRealityCaptureFramework/Private/MrcCalibrationData.cpp b/Engine/Plugins/Runtime/MixedRealityCaptureFramework/Source/MixedRealityCaptureFramework/Private/MrcCalibrationData.cpp
--- a/Engine/Plugins/Runtime/MixedRealityCaptureFramework/Source/MixedRealityCaptureFramework/Private/MrcCalibrationData.cpp
+++ b/Engine/Plugins/Runtime/MixedRealityCaptureFramework/Source/MixedRealityCaptureFramework/Private/MrcCalibrationData.cpp
@@ -12,5 +12,5 @@ UMrcCalibrationSaveGame::UMrcCalibrationSaveGame(const FObjectInitializer& Objec
 {
 	SaveSlotName = TEXT("MrcCalibration");
 	UserIndex = 0;
-	ConfigurationSaveVersion = 0;
+	ConfigurationSaveVersion = 1;
 }
\ No newline at end of file
