Mirroring CL 4084827/4085071 from Dev-VR

MRC - Adding CVar for overriding the tracking latency (makes development easier in the calibration process).

diff --git a/Engine/Plugins/Runtime/MixedRealityCaptureFramework/Source/MixedRealityCaptureFramework/Private/MixedRealityCaptureComponent.cpp b/Engine/Plugins/Runtime/MixedRealityCaptureFramework/Source/MixedRealityCaptureFramework/Private/MixedRealityCaptureComponent.cpp
--- a/Engine/Plugins/Runtime/MixedRealityCaptureFramework/Source/MixedRealityCaptureFramework/Private/MixedRealityCaptureComponent.cpp
+++ b/Engine/Plugins/Runtime/MixedRealityCaptureFramework/Source/MixedRealityCaptureFramework/Private/MixedRealityCaptureComponent.cpp
@@ -78,6 +78,7 @@ bool TMulticastCVarCommand<bool, int32>::GetValue() const
 
 typedef TMulticastCVarCommand<bool, int32> FMulticastBoolCVar;
 typedef TMulticastCVarCommand<float> FMulticastFloatCVar;
+typedef TMulticastCVarCommand<int32> FMulticastIntCVar;
 
 /* MRCaptureComponent_Impl
  *****************************************************************************/
@@ -114,6 +115,12 @@ namespace MRCaptureComponent_Impl
 		TEXT("When set to be greater than zero, MixedRealityCaptures will use this for the FOV instead of what was previously set.")
 	);
 
+	static FMulticastIntCVar TrackingLatencyOverride(
+		TEXT("mrc.TrackingLatencyOverride"),
+		0,
+		TEXT("When set to be greater than zero, MixedRealityCaptures will use this for their TrackingLatency instead of what's set. The higher the value (in ms), the more delay there will be introduced to tracked components.")
+	);
+
 	static const FName DefaultDistortionMapParamName(TEXT("DistortionDisplacementMap"));
 	static const EHMDTrackingOrigin::Type InvalidOriginType = (EHMDTrackingOrigin::Type)0xFF;
 
@@ -825,6 +832,12 @@ void UMixedRealityCaptureComponent::SetLensDistortionParameters(const FMrcLensDi
 	}
 }
 
+//------------------------------------------------------------------------------
+int32 UMixedRealityCaptureComponent::GetTrackingDelay() const
+{
+	return (MRCaptureComponent_Impl::TrackingLatencyOverride > 0) ? (int32)MRCaptureComponent_Impl::TrackingLatencyOverride : TrackingLatency;
+}
+
 //------------------------------------------------------------------------------
 void UMixedRealityCaptureComponent::SetTrackingDelay(int32 DelayMS)
 {
@@ -1076,7 +1089,7 @@ void UMixedRealityCaptureComponent::FillOutCalibrationData(UMrcCalibrationData*
 		{
 			Dst->CompositingData.CaptureDeviceURL = CaptureFeedRef;
 			Dst->CompositingData.DepthOffset = ProjectionDepthOffset;
-			Dst->CompositingData.TrackingLatency = TrackingLatency;
+			Dst->CompositingData.TrackingLatency = GetTrackingDelay();
 			Dst->CompositingData.VideoProcessingParams = VideoProcessingParams;
 		}
 		// garbage matte
diff --git a/Engine/Plugins/Runtime/MixedRealityCaptureFramework/Source/MixedRealityCaptureFramework/Private/MrcLatencyViewExtension.cpp b/Engine/Plugins/Runtime/MixedRealityCaptureFramework/Source/MixedRealityCaptureFramework/Private/MrcLatencyViewExtension.cpp
--- a/Engine/Plugins/Runtime/MixedRealityCaptureFramework/Source/MixedRealityCaptureFramework/Private/MrcLatencyViewExtension.cpp
+++ b/Engine/Plugins/Runtime/MixedRealityCaptureFramework/Source/MixedRealityCaptureFramework/Private/MrcLatencyViewExtension.cpp
@@ -27,7 +27,7 @@ static uint32 MRLatencyViewExtension_Impl::GetDesiredDelay(TWeakObjectPtr<UMixed
 	uint32 DesiredDelay = CVarMotionCaptureLatencyOverride.GetValueOnGameThread();
 	if (DesiredDelay <= 0 && Target.IsValid())
 	{
-		DesiredDelay = Target->TrackingLatency;
+		DesiredDelay = Target->GetTrackingDelay();
 	}
 	return DesiredDelay;
 }
@@ -124,7 +124,7 @@ uint32 FMrcLatencyViewExtension::GetDesiredDelay() const
 	uint32 DesiredDelay = MRLatencyViewExtension_Impl::CVarMotionCaptureLatencyOverride.GetValueOnGameThread();
 	if (DesiredDelay <= 0 && Owner.IsValid())
 	{
-		DesiredDelay = Owner->TrackingLatency;
+		DesiredDelay = Owner->GetTrackingDelay();
 	}
 	return DesiredDelay;
 }
diff --git a/Engine/Plugins/Runtime/MixedRealityCaptureFramework/Source/MixedRealityCaptureFramework/Public/MixedRealityCaptureComponent.h b/Engine/Plugins/Runtime/MixedRealityCaptureFramework/Source/MixedRealityCaptureFramework/Public/MixedRealityCaptureComponent.h
--- a/Engine/Plugins/Runtime/MixedRealityCaptureFramework/Source/MixedRealityCaptureFramework/Public/MixedRealityCaptureComponent.h
+++ b/Engine/Plugins/Runtime/MixedRealityCaptureFramework/Source/MixedRealityCaptureFramework/Public/MixedRealityCaptureComponent.h
@@ -58,7 +58,7 @@ public:
 	UTextureRenderTarget2D* GarbageMatteCaptureTextureTarget;
 
 	/** Millisecond delay to apply to motion controller components when rendering to the capture view (to better align with latent camera feeds) */
-	UPROPERTY(BlueprintReadWrite, Config, BlueprintSetter=SetTrackingDelay, Category=Composition, meta=(ClampMin="0", UIMin="0"))
+	UPROPERTY(BlueprintReadWrite, Config, BlueprintGetter=GetTrackingDelay, BlueprintSetter=SetTrackingDelay, Category=Composition, meta=(ClampMin="0", UIMin="0"))
 	int32 TrackingLatency;
 
 	/** Determines if this component should attempt to load the default MR calibration file on initialization */
@@ -156,6 +156,9 @@ public:
 	UFUNCTION(BlueprintSetter)
 	void SetLensDistortionParameters(const FMrcLensDistortion& ModelRef);
 
+	UFUNCTION(BlueprintGetter)
+	int32 GetTrackingDelay() const;
+
 	UFUNCTION(BlueprintSetter)
 	void SetTrackingDelay(int32 DelayMS);
 
