Mirroring CL 4090670/4094620 from Dev-VR

MR Calibration - Fixing an issue where the alignment models would disappear.
MR Calibration - Saving garbage mattes when exiting using escape.

diff --git a/Engine/Plugins/Runtime/MixedRealityCaptureFramework/Source/MixedRealityCaptureFramework/Public/MrcUtilLibrary.h b/Engine/Plugins/Runtime/MixedRealityCaptureFramework/Source/MixedRealityCaptureFramework/Public/MrcUtilLibrary.h
--- a/Engine/Plugins/Runtime/MixedRealityCaptureFramework/Source/MixedRealityCaptureFramework/Public/MrcUtilLibrary.h
+++ b/Engine/Plugins/Runtime/MixedRealityCaptureFramework/Source/MixedRealityCaptureFramework/Public/MrcUtilLibrary.h
@@ -25,7 +25,7 @@ public:
 	static bool IsMixedRealityCaptureBroadcasting();
 
 	/**
-	 * Return true if the system is sending the capture texture to the spectator screen.
+	 * Toggle whether the capture system is sending the capture texture to the spectator screen or not.
 	 */
 	UFUNCTION(BlueprintCallable, Category = "MixedRealityCapture|Utils")
 	static void SetMixedRealityCaptureBroadcasting(bool bEnable);
