Fix UMrcUtilLibrary::GetHMDCameraComponent compile error
#rnx

diff --git a/Engine/Plugins/Runtime/MixedRealityCaptureFramework/Source/MixedRealityCaptureFramework/Private/MrcUtilLibrary.cpp b/Engine/Plugins/Runtime/MixedRealityCaptureFramework/Source/MixedRealityCaptureFramework/Private/MrcUtilLibrary.cpp
--- a/Engine/Plugins/Runtime/MixedRealityCaptureFramework/Source/MixedRealityCaptureFramework/Private/MrcUtilLibrary.cpp
+++ b/Engine/Plugins/Runtime/MixedRealityCaptureFramework/Source/MixedRealityCaptureFramework/Private/MrcUtilLibrary.cpp
@@ -145,13 +145,12 @@ UCameraComponent* UMrcUtilLibrary::GetHMDCameraComponent(const APawn* PlayerPawn
 					HMDCam = Camera;
 					break;
 				}
+				else if (HMDCam == nullptr)
+				{
+					HMDCam = Camera;
+				}
 			}
 		}
-
-		if (HMDCam == nullptr && CameraComponents.Num() > 0)
-		{
-			HMDCam = CameraComponents[0];
-		}
 	}
 	return HMDCam;
 }
