CIS Fix - Fix from bad checkin (didn't save changes after renaming function... opps!)

diff --git a/Engine/Plugins/Runtime/MixedRealityCaptureFramework/Source/MixedRealityCaptureFramework/Private/MixedRealityCaptureComponent.cpp b/Engine/Plugins/Runtime/MixedRealityCaptureFramework/Source/MixedRealityCaptureFramework/Private/MixedRealityCaptureComponent.cpp
--- a/Engine/Plugins/Runtime/MixedRealityCaptureFramework/Source/MixedRealityCaptureFramework/Private/MixedRealityCaptureComponent.cpp
+++ b/Engine/Plugins/Runtime/MixedRealityCaptureFramework/Source/MixedRealityCaptureFramework/Private/MixedRealityCaptureComponent.cpp
@@ -1089,7 +1089,7 @@ void UMixedRealityCaptureComponent::FillOutCalibrationData(UMrcCalibrationData*
 		{
 			Dst->CompositingData.CaptureDeviceURL = CaptureFeedRef;
 			Dst->CompositingData.DepthOffset = ProjectionDepthOffset;
-			Dst->CompositingData.TrackingLatency = GetTrackingLatency();
+			Dst->CompositingData.TrackingLatency = GetTrackingDelay();
 			Dst->CompositingData.VideoProcessingParams = VideoProcessingParams;
 		}
 		// garbage matte
diff --git a/Engine/Plugins/Runtime/MixedRealityCaptureFramework/Source/MixedRealityCaptureFramework/Private/MrcLatencyViewExtension.cpp b/Engine/Plugins/Runtime/MixedRealityCaptureFramework/Source/MixedRealityCaptureFramework/Private/MrcLatencyViewExtension.cpp
--- a/Engine/Plugins/Runtime/MixedRealityCaptureFramework/Source/MixedRealityCaptureFramework/Private/MrcLatencyViewExtension.cpp
+++ b/Engine/Plugins/Runtime/MixedRealityCaptureFramework/Source/MixedRealityCaptureFramework/Private/MrcLatencyViewExtension.cpp
@@ -27,7 +27,7 @@ static uint32 MRLatencyViewExtension_Impl::GetDesiredDelay(TWeakObjectPtr<UMixed
 	uint32 DesiredDelay = CVarMotionCaptureLatencyOverride.GetValueOnGameThread();
 	if (DesiredDelay <= 0 && Target.IsValid())
 	{
-		DesiredDelay = Target->GetTrackingLatency();
+		DesiredDelay = Target->GetTrackingDelay();
 	}
 	return DesiredDelay;
 }
@@ -124,7 +124,7 @@ uint32 FMrcLatencyViewExtension::GetDesiredDelay() const
 	uint32 DesiredDelay = MRLatencyViewExtension_Impl::CVarMotionCaptureLatencyOverride.GetValueOnGameThread();
 	if (DesiredDelay <= 0 && Owner.IsValid())
 	{
-		DesiredDelay = Owner->GetTrackingLatency();
+		DesiredDelay = Owner->GetTrackingDelay();
 	}
 	return DesiredDelay;
 }
