Adding a null check to the MRC plugin for levels that fail to open.

diff --git a/Engine/Plugins/Runtime/MixedRealityCaptureFramework/Source/MixedRealityCaptureFramework/Private/MrcFrameworkModule.cpp b/Engine/Plugins/Runtime/MixedRealityCaptureFramework/Source/MixedRealityCaptureFramework/Private/MrcFrameworkModule.cpp
--- a/Engine/Plugins/Runtime/MixedRealityCaptureFramework/Source/MixedRealityCaptureFramework/Private/MrcFrameworkModule.cpp
+++ b/Engine/Plugins/Runtime/MixedRealityCaptureFramework/Source/MixedRealityCaptureFramework/Private/MrcFrameworkModule.cpp
@@ -80,7 +80,7 @@ void FMrcFrameworkModule::ShutdownModule()
 void FMrcFrameworkModule::OnWorldCreated(UWorld* NewWorld)
 {
 #if WITH_EDITORONLY_DATA
-	const bool bIsGameInst = !IsRunningCommandlet() && NewWorld->IsGameWorld();
+	const bool bIsGameInst = !IsRunningCommandlet() && NewWorld && NewWorld->IsGameWorld();
 	if (bIsGameInst)
 #endif 
 	{
